Added ES2 Deprecation warning to UI

diff --git a/Engine/Source/Runtime/MaterialShaderQualitySettings/Private/MaterialShaderQualitySettingsCustomization.cpp b/Engine/Source/Runtime/MaterialShaderQualitySettings/Private/MaterialShaderQualitySettingsCustomization.cpp
--- a/Engine/Source/Runtime/MaterialShaderQualitySettings/Private/MaterialShaderQualitySettingsCustomization.cpp
+++ b/Engine/Source/Runtime/MaterialShaderQualitySettings/Private/MaterialShaderQualitySettingsCustomization.cpp
@@ -178,6 +178,26 @@ void FMaterialShaderQualitySettingsCustomization::CustomizeDetails(IDetailLayout
 		QualityOverrideListSource.Add(MakeShareable(new FShaderQualityOverridesListItem(DisplayName, OverrideHandles, EnabledHandles)));
 	}
 
+	TArray<TWeakObjectPtr<UObject>> CustomizedObjects;
+	DetailLayout.GetObjectsBeingCustomized(CustomizedObjects);
+	if (CustomizedObjects.Num() && CustomizedObjects[0]->GetName() == TEXT("ForwardShadingQuality_GLSL_ES2"))
+	{
+		TSharedPtr<SWidget> ES2Warning = SNew(SBorder)
+			.BorderBackgroundColor(FLinearColor::Red)
+			.BorderImage(FEditorStyle::GetBrush("ToolPanel.LightGroupBorder"))
+			.Padding(8.0f)
+			[
+				SNew(STextBlock)
+				.Text(LOCTEXT("ES2Deprecation", "Warning: This feature level is deprecated. Please use ES3.1"))
+			];
+
+		ForwardRenderingCategory.AddCustomRow(LOCTEXT("Warning", "Warning"), false)
+			.WholeRowWidget
+			[
+				ES2Warning.ToSharedRef()
+			];
+	}
+
 	ForwardRenderingCategory.AddCustomRow(LOCTEXT("ForwardRenderingMaterialOverrides", "Forward Rendering Material Overrides"))
 		[
 			SAssignNew(MaterialQualityOverridesListView, SMaterialQualityOverridesListView)
