Fix a VS2019 warning (unknown enum size for EShaderPlatform in quality settings).

#rnx

diff --git a/Engine/Source/Runtime/MaterialShaderQualitySettings/Classes/ShaderPlatformQualitySettings.h b/Engine/Source/Runtime/MaterialShaderQualitySettings/Classes/ShaderPlatformQualitySettings.h
--- a/Engine/Source/Runtime/MaterialShaderQualitySettings/Classes/ShaderPlatformQualitySettings.h
+++ b/Engine/Source/Runtime/MaterialShaderQualitySettings/Classes/ShaderPlatformQualitySettings.h
@@ -6,6 +6,7 @@
 #include "UObject/ObjectMacros.h"
 #include "UObject/Object.h"
 #include "SceneTypes.h"
+#include "RHIDefinitions.h"
 #include "ShaderPlatformQualitySettings.generated.h"
 
 /**
@@ -65,7 +66,7 @@ public:
 	UPROPERTY(EditAnywhere, Config, Meta = (DisplayName = "Cascade shadow mapping quality"), Category = "Quality")
 	EMobileCSMQuality MobileCSMQuality;
 
-	MATERIALSHADERQUALITYSETTINGS_API bool CanOverride(enum EShaderPlatform ShaderPlatform) const;
+	MATERIALSHADERQUALITYSETTINGS_API bool CanOverride(EShaderPlatform ShaderPlatform) const;
 	bool HasAnyOverridesSet() const;
 };
 
